[JsonBlueprint] Removed unused world context requirement from two json loading functions

diff --git a/Engine/Plugins/JsonBlueprintUtilities/Source/JsonBlueprintUtilities/Private/JsonBlueprintFunctionLibrary.cpp b/Engine/Plugins/JsonBlueprintUtilities/Source/JsonBlueprintUtilities/Private/JsonBlueprintFunctionLibrary.cpp
--- a/Engine/Plugins/JsonBlueprintUtilities/Source/JsonBlueprintUtilities/Private/JsonBlueprintFunctionLibrary.cpp
+++ b/Engine/Plugins/JsonBlueprintUtilities/Source/JsonBlueprintUtilities/Private/JsonBlueprintFunctionLibrary.cpp
@@ -10,7 +10,6 @@
 #define LOCTEXT_NAMESPACE "JsonBlueprintFunctionLibrary"
 
 bool UJsonBlueprintFunctionLibrary::FromString(
-	UObject* WorldContextObject,
 	const FString& JsonString,
 	FJsonObjectWrapper& OutJsonObject)
 {
@@ -18,7 +17,6 @@ bool UJsonBlueprintFunctionLibrary::FromString(
 }
 
 bool UJsonBlueprintFunctionLibrary::FromFile(
-	UObject* WorldContextObject,	
 	const FFilePath& File,
 	FJsonObjectWrapper& OutJsonObject)
 {
@@ -35,7 +33,7 @@ bool UJsonBlueprintFunctionLibrary::FromFile(
 		return false;
 	}
 	
-	return FromString(WorldContextObject, JsonString, OutJsonObject);
+	return FromString(JsonString, OutJsonObject);
 }
 
 bool UJsonBlueprintFunctionLibrary::ToString(const FJsonObjectWrapper& JsonObject, FString& OutJsonString)
diff --git a/Engine/Plugins/JsonBlueprintUtilities/Source/JsonBlueprintUtilities/Public/JsonBlueprintFunctionLibrary.h b/Engine/Plugins/JsonBlueprintUtilities/Source/JsonBlueprintUtilities/Public/JsonBlueprintFunctionLibrary.h
--- a/Engine/Plugins/JsonBlueprintUtilities/Source/JsonBlueprintUtilities/Public/JsonBlueprintFunctionLibrary.h
+++ b/Engine/Plugins/JsonBlueprintUtilities/Source/JsonBlueprintUtilities/Public/JsonBlueprintFunctionLibrary.h
@@ -16,12 +16,12 @@ class JSONBLUEPRINTUTILITIES_API UJsonBlueprintFunctionLibrary final : public UB
 
 public:
 	/** Creates a JsonObject from the provided Json string. */
-	UFUNCTION(BlueprintCallable, Category="Json", meta = (WorldContext="WorldContextObject",  HidePin = "WorldContextObject", DefaultToSelf = "WorldContextObject", DisplayName="Load Json from String"))
-	static UPARAM(DisplayName="Success") bool FromString(UObject* WorldContextObject, const FString& JsonString, UPARAM(DisplayName="JsonObject") FJsonObjectWrapper& OutJsonObject);
+	UFUNCTION(BlueprintCallable, Category="Json", meta = (DisplayName="Load Json from String"))
+	static UPARAM(DisplayName="Success") bool FromString(const FString& JsonString, UPARAM(DisplayName="JsonObject") FJsonObjectWrapper& OutJsonObject);
 
 	/** Creates a JsonObject from the provided Json file. */
-	UFUNCTION(BlueprintCallable, Category="Json", meta = (WorldContext="WorldContextObject",  HidePin = "WorldContextObject", DefaultToSelf = "WorldContextObject", DisplayName="Load Json from File"))
-	static UPARAM(DisplayName="Success") bool FromFile(UObject* WorldContextObject, const FFilePath& File, UPARAM(DisplayName="JsonObject") FJsonObjectWrapper& OutJsonObject);
+	UFUNCTION(BlueprintCallable, Category="Json", meta = (DisplayName="Load Json from File"))
+	static UPARAM(DisplayName="Success") bool FromFile(const FFilePath& File, UPARAM(DisplayName="JsonObject") FJsonObjectWrapper& OutJsonObject);
 
 	/** Creates a Json string from the provided JsonObject. */
 	UFUNCTION(BlueprintCallable, Category="Json", meta = (DisplayName = "Get Json String"))
